#UE4 make PostSplashScreen entirely disabled on pc only and upadate LiveCodingConsole binary to remove dependency on it.

diff --git a/Engine/Plugins/Runtime/PostSplashScreen/PostSplashScreen.uplugin b/Engine/Plugins/Runtime/PostSplashScreen/PostSplashScreen.uplugin
--- a/Engine/Plugins/Runtime/PostSplashScreen/PostSplashScreen.uplugin
+++ b/Engine/Plugins/Runtime/PostSplashScreen/PostSplashScreen.uplugin
@@ -14,6 +14,9 @@
 	"CanContainContent": true,
 	"IsBetaVersion": false,
 	"Installed": false,
+	"SupportedTargetPlatforms": [
+		"XboxOne"
+	],
 	"Modules": [
 		{
 			"Name": "PostSplashScreen",
