Fix custom splash.

diff --git a/Engine/Plugins/Runtime/PostSplashScreen/Source/PostSplashScreen/Private/CustomSplashScreen.cpp b/Engine/Plugins/Runtime/PostSplashScreen/Source/PostSplashScreen/Private/CustomSplashScreen.cpp
--- a/Engine/Plugins/Runtime/PostSplashScreen/Source/PostSplashScreen/Private/CustomSplashScreen.cpp
+++ b/Engine/Plugins/Runtime/PostSplashScreen/Source/PostSplashScreen/Private/CustomSplashScreen.cpp
@@ -13,9 +13,9 @@
 #define CONFIG_BASENAME TEXT("PostSplashScreen")
 FString GPostSplashScreenIni;
 
-bool FCustomSplashScreen::Init()
+bool FCustomSplashScreen::Init(TFunctionRef<TSharedPtr<IAnalyticsProviderET>()> AnalyticsFactory)
 {
-	FPreLoadScreenBase::Init();
+	FPreLoadScreenBase::Init(AnalyticsFactory);
 	
 	SetPluginName("PostSplashScreen");
 	InitSettingsFromConfig("PostSplashScreen");
diff --git a/Engine/Plugins/Runtime/PostSplashScreen/Source/PostSplashScreen/Public/CustomSplashScreen.h b/Engine/Plugins/Runtime/PostSplashScreen/Source/PostSplashScreen/Public/CustomSplashScreen.h
--- a/Engine/Plugins/Runtime/PostSplashScreen/Source/PostSplashScreen/Public/CustomSplashScreen.h
+++ b/Engine/Plugins/Runtime/PostSplashScreen/Source/PostSplashScreen/Public/CustomSplashScreen.h
@@ -17,7 +17,7 @@ public:
 	
     /*** IPreLoadScreen Implementation ***/
 	virtual void Tick(float DeltaTime) override;
-    virtual bool Init() override;
+    virtual bool Init(TFunctionRef<TSharedPtr<IAnalyticsProviderET>()> AnalyticsFactory) override;
 
     //Override to make sure this is always an EarlyStartupScreen
     virtual EPreLoadScreenTypes GetPreLoadScreenType() const override { return EPreLoadScreenTypes::CustomSplashScreen; }
