Merging CL#9483124

#PostSplashScreen CIS Fix

diff --git a/Engine/Plugins/Runtime/PostSplashScreen/Source/PostSplashScreen/Private/PostSplashScreenModule.cpp b/Engine/Plugins/Runtime/PostSplashScreen/Source/PostSplashScreen/Private/PostSplashScreenModule.cpp
--- a/Engine/Plugins/Runtime/PostSplashScreen/Source/PostSplashScreen/Private/PostSplashScreenModule.cpp
+++ b/Engine/Plugins/Runtime/PostSplashScreen/Source/PostSplashScreen/Private/PostSplashScreenModule.cpp
@@ -23,7 +23,7 @@ public:
 		if (!GIsEditor && FApp::CanEverRender() && FPreLoadScreenManager::Get())
 		{
 			CustomSplashScreen = MakeShared<FCustomSplashScreen>();
-			if (CustomSplashScreen->Init())
+			if (CustomSplashScreen->Init([]()->TSharedPtr<IAnalyticsProviderET>{ return TSharedPtr<IAnalyticsProviderET>(); }))
 			{
 				FPreLoadScreenManager::Get()->RegisterPreLoadScreen(CustomSplashScreen);
 			}
